Virtual Production Templates.  Updating Media Profiles, enabling Timecode Syncronizer.  Cleaning up ini files.

diff --git a/Templates/TP_ME_BlankBP/Config/DefaultEngine.ini b/Templates/TP_ME_BlankBP/Config/DefaultEngine.ini
--- a/Templates/TP_ME_BlankBP/Config/DefaultEngine.ini
+++ b/Templates/TP_ME_BlankBP/Config/DefaultEngine.ini
@@ -5,13 +5,16 @@ AppliedTargetedHardwareClass=Desktop
 DefaultGraphicsPerformance=Maximum
 AppliedDefaultGraphicsPerformance=Maximum
 
+
 [/Script/Engine.RendererSettings]
 ; exposure control and tone mapping settings
 r.Tonemapper.Sharpen=2
 r.TemporalAA.Upsampling=True
 r.DefaultFeature.AutoExposure.ExtendDefaultLuminanceRange=True
-;r.DefaultFeature.AutoExposure=False
-;r.DefaultFeature.AutoExposure.Method=0
+r.DefaultFeature.Bloom=False
+r.DefaultFeature.AutoExposure=False
+r.DefaultFeature.MotionBlur=False
+r.SceneRenderTargetResizeMethod=2
 
 ; dynamic gi settings
 r.LightPropagationVolume=0
@@ -68,9 +71,7 @@ ScoutingSubsystemEdititorUtilityActorClassPath=/VirtualProductionUtilities/Virtu
 
 
 ; Enables WebControl API
-
-;WebControl.EnableServerOnStartup=True
-;WebControl.EnableServerOnStartup=1
+WebControl.EnableServerOnStartup=1
 
 ;Adds Remote Session for Vcam
 [RemoteSession]
@@ -100,5 +101,12 @@ GameViewportClientClassName=/Script/DisplayCluster.DisplayClusterViewportClient
 [/Script/DisplayClusterEditor.DisplayClusterEditorSettings]
 bEnabled=True
 
+[/Script/WindowsTargetPlatform.WindowsTargetSettings]
+;Ndisplay Requirements
+DefaultGraphicsRHI=DefaultGraphicsRHI_DX11
++TargetedRHIs=PCD3D_SM5
+-TargetedRHIs=PCD3D_SM4
+-TargetedRHIs=GLSL_430
+-TargetedRHIs=GLSL_150
 
 
diff --git a/Templates/TP_ME_BlankBP/TP_ME_BlankBP.uproject b/Templates/TP_ME_BlankBP/TP_ME_BlankBP.uproject
--- a/Templates/TP_ME_BlankBP/TP_ME_BlankBP.uproject
+++ b/Templates/TP_ME_BlankBP/TP_ME_BlankBP.uproject
@@ -14,11 +14,11 @@
 		},
 		{
 			"Name": "OculusVR",
-			"Enabled": false
+			"Enabled": true
 		},
 		{
 			"Name": "SteamVR",
-			"Enabled": false
+			"Enabled": true
 		},
 		{
 			"Name": "PythonScriptPlugin",
@@ -107,6 +107,14 @@
 		{
 			"Name": "PixelStreaming",
 			"Enabled": true
+		},
+			{
+			"Name": "TimecodeSynchronizer",
+			"Enabled": true
+		},
+		{
+			"Name": "VariantManager",
+			"Enabled": true
 		}
 	]
 }
\ No newline at end of file
diff --git a/Templates/TP_ME_VProdBP/Config/DefaultEngine.ini b/Templates/TP_ME_VProdBP/Config/DefaultEngine.ini
--- a/Templates/TP_ME_VProdBP/Config/DefaultEngine.ini
+++ b/Templates/TP_ME_VProdBP/Config/DefaultEngine.ini
@@ -5,13 +5,24 @@ AppliedTargetedHardwareClass=Desktop
 DefaultGraphicsPerformance=Maximum
 AppliedDefaultGraphicsPerformance=Maximum
 
+[/Script/WindowsTargetPlatform.WindowsTargetSettings]
+;Ndisplay Requirements
+DefaultGraphicsRHI=DefaultGraphicsRHI_DX11
++TargetedRHIs=PCD3D_SM5
+-TargetedRHIs=PCD3D_SM4
+-TargetedRHIs=GLSL_430
+-TargetedRHIs=GLSL_150
+
 [/Script/Engine.RendererSettings]
 ; exposure control and tone mapping settings
 r.Tonemapper.Sharpen=2
 r.TemporalAA.Upsampling=True
 r.DefaultFeature.AutoExposure.ExtendDefaultLuminanceRange=True
-;r.DefaultFeature.AutoExposure=False
-;r.DefaultFeature.AutoExposure.Method=0
+r.DefaultFeature.Bloom=False
+r.DefaultFeature.AutoExposure=False
+r.DefaultFeature.MotionBlur=False
+r.SceneRenderTargetResizeMethod=2
+
 
 ;dynamic gi settings
 r.LightPropagationVolume=0
@@ -28,7 +39,7 @@ r.GBufferFormat=3
 ; shader settings
 r.PostProcessing.PropagateAlpha=1
 r.SupportSkyAtmosphereAffectsHeightFog=True
-r.DefaultFeature.LensFlare=True
+r.DefaultFeature.LensFlare=False
 r.ClearCoatNormal=False
 r.NormalMapsForStaticLighting=False
 
@@ -57,10 +68,8 @@ NativeAudioOut=False
 [/Script/VPUtilitiesEditor.VPUtilitiesEditorSettings]
 VirtualScoutingUI=/VirtualProductionUtilities/Editor/VirtualScoutingWidget.VirtualScoutingWidget_C
 
-;Enables WebControl API
-
-;WebControl.EnableServerOnStartup=True
-;WebControl.EnableServerOnStartup=1
+; Enables WebControl API
+WebControl.EnableServerOnStartup=1
 
 ;Adds Remote Session for Vcam
 [RemoteSession]
@@ -89,15 +98,23 @@ DefaultSaveSessionAs=
 ClientSettings=(DisplayName="",AvatarColor=(R=1.000000,G=1.000000,B=1.000000,A=1.000000),DesktopAvatarActorClass=/ConcertSyncClient/DesktopPresence.DesktopPresence_C,VRAvatarActorClass=/ConcertSyncClient/VRPresence.VRPresence_C,DiscoveryTimeoutSeconds=5,SessionTickFrequencySeconds=1,LatencyCompensationMs=0.000000,Tags=)
 EndpointSettings=(bEnableLogging=False,PurgeProcessedMessageDelaySeconds=30,RemoteEndpointTimeoutSeconds=60)
 
-;Setup for Ndisplay
+
+;Setup for Ndisplay Cluster
 [/Script/Engine.Engine]
 GameEngine=/Script/DisplayCluster.DisplayClusterGameEngine
 UnrealEdEngine=/Script/DisplayClusterEditor.DisplayClusterEditorEngine
 GameViewportClientClassName=/Script/DisplayCluster.DisplayClusterViewportClient
 
-;Setup for Multiuser
+;Setup for Ndisplay Cluster
 [/Script/DisplayClusterEditor.DisplayClusterEditorSettings]
 bEnabled=True
 
+[/Script/WindowsTargetPlatform.WindowsTargetSettings]
+;Ndisplay Requirements
+DefaultGraphicsRHI=DefaultGraphicsRHI_DX11
++TargetedRHIs=PCD3D_SM5
+-TargetedRHIs=PCD3D_SM4
+-TargetedRHIs=GLSL_430
+-TargetedRHIs=GLSL_150
 
 
diff --git a/Templates/TP_ME_VProdBP/TP_ME_VprodBP.uproject b/Templates/TP_ME_VProdBP/TP_ME_VprodBP.uproject
--- a/Templates/TP_ME_VProdBP/TP_ME_VprodBP.uproject
+++ b/Templates/TP_ME_VProdBP/TP_ME_VprodBP.uproject
@@ -14,11 +14,11 @@
 		},
 		{
 			"Name": "OculusVR",
-			"Enabled": false
+			"Enabled": true
 		},
 		{
 			"Name": "SteamVR",
-			"Enabled": false
+			"Enabled": true
 		},
 		{
 			"Name": "PythonScriptPlugin",
@@ -107,6 +107,14 @@
 		{
 			"Name": "PixelStreaming",
 			"Enabled": true
+		},
+			{
+			"Name": "TimecodeSynchronizer",
+			"Enabled": true
+		},
+		{
+			"Name": "VariantManager",
+			"Enabled": true
 		}
 	]
 }
\ No newline at end of file
