Virtual Production Templates.  Update to the Game Map.  Added Composure Linear to SRGB Output Material.  Updated Project files.

diff --git a/Templates/TP_ME_BlankBP/Config/DefaultEngine.ini b/Templates/TP_ME_BlankBP/Config/DefaultEngine.ini
--- a/Templates/TP_ME_BlankBP/Config/DefaultEngine.ini
+++ b/Templates/TP_ME_BlankBP/Config/DefaultEngine.ini
@@ -35,13 +35,19 @@ r.DefaultFeature.LensFlare=True
 r.ClearCoatNormal=False
 r.NormalMapsForStaticLighting=False
 
+>>>> ORIGINAL //UE4/Release-4.24/Templates/TP_ME_BlankBP/Config/DefaultEngine.ini#6
+; Sequencer Take Recorder
+Sequencer.AlwayRecordLiveLinkTimecode 1
+==== THEIRS //UE4/Release-4.24/Templates/TP_ME_BlankBP/Config/DefaultEngine.ini#7
 ; Sequencer Take Recorder
 Sequencer.AlwayRecordLiveLinkTimecode=1
+==== YOURS //andy.blondin_4-24_Release/Templates/TP_ME_BlankBP/Config/DefaultEngine.ini
+<<<<
 
 
 [/Script/EngineSettings.GameMapsSettings]
 ;defines what level to be loaded by default in the editor and at playback time
-GameDefaultMap=/Game/Main.Main
+GameDefaultMap=/Game/Ndisplay.Ndisplay
 EditorStartupMap=/Game/Main.Main
 
 [/Script/UnrealEd.UnrealEdEngine]
diff --git a/Templates/TP_ME_BlankBP/TP_ME_BlankBP.uproject b/Templates/TP_ME_BlankBP/TP_ME_BlankBP.uproject
--- a/Templates/TP_ME_BlankBP/TP_ME_BlankBP.uproject
+++ b/Templates/TP_ME_BlankBP/TP_ME_BlankBP.uproject
@@ -14,13 +14,23 @@
 		},
 		{
 			"Name": "OculusVR",
-			"Enabled": true
+			"Enabled": true,
+			"SupportedTargetPlatforms": [
+				"Win32",
+				"Win64",
+				"Android"
+			]
 		},
 		{
 			"Name": "SteamVR",
-			"Enabled": true
+			"Enabled": true,
+			"SupportedTargetPlatforms": [
+				"Win32",
+				"Win64",
+				"Linux",
+				"Mac"
+			]
 		},
-		{
 			"Name": "PythonScriptPlugin",
 			"Enabled": true
 		},
diff --git a/Templates/TP_ME_VProdBP/Config/DefaultEngine.ini b/Templates/TP_ME_VProdBP/Config/DefaultEngine.ini
--- a/Templates/TP_ME_VProdBP/Config/DefaultEngine.ini
+++ b/Templates/TP_ME_VProdBP/Config/DefaultEngine.ini
@@ -43,13 +43,19 @@ r.DefaultFeature.LensFlare=False
 r.ClearCoatNormal=False
 r.NormalMapsForStaticLighting=False
 
+>>>> ORIGINAL //UE4/Release-4.24/Templates/TP_ME_VProdBP/Config/DefaultEngine.ini#3
+; Sequencer Take Recorder
+Sequencer.AlwayRecordLiveLinkTimecode 1
+==== THEIRS //UE4/Release-4.24/Templates/TP_ME_VProdBP/Config/DefaultEngine.ini#4
 ; Sequencer Take Recorder
 Sequencer.AlwayRecordLiveLinkTimecode=1
+==== YOURS //andy.blondin_4-24_Release/Templates/TP_ME_VProdBP/Config/DefaultEngine.ini
+<<<<
 
 
 [/Script/EngineSettings.GameMapsSettings]
 ;defines what level to be loaded by default in the editor and at playback time
-GameDefaultMap=/Game/VProdProject/Main.Main
+GameDefaultMap=/Game/VProdProject/Ndisplay.Ndisplay
 EditorStartupMap=/Game/VProdProject/Main.Main
 
 [/Script/UnrealEd.UnrealEdEngine]
diff --git a/Templates/TP_ME_VProdBP/TP_ME_VprodBP.uproject b/Templates/TP_ME_VProdBP/TP_ME_VprodBP.uproject
--- a/Templates/TP_ME_VProdBP/TP_ME_VprodBP.uproject
+++ b/Templates/TP_ME_VProdBP/TP_ME_VprodBP.uproject
@@ -14,11 +14,22 @@
 		},
 		{
 			"Name": "OculusVR",
-			"Enabled": true
+			"Enabled": true,
+			"SupportedTargetPlatforms": [
+				"Win32",
+				"Win64",
+				"Android"
+			]
 		},
 		{
 			"Name": "SteamVR",
-			"Enabled": true
+			"Enabled": true,
+			"SupportedTargetPlatforms": [
+				"Win32",
+				"Win64",
+				"Linux",
+				"Mac"
+			]
 		},
 		{
 			"Name": "PythonScriptPlugin",
@@ -108,7 +119,7 @@
 			"Name": "PixelStreaming",
 			"Enabled": true
 		},
-			{
+		{
 			"Name": "TimecodeSynchronizer",
 			"Enabled": true
 		},
