Virutal Production Project Template.  Move Maps.  Removed Warnings

diff --git a/Templates/TP_ME_BlankBP/Config/DefaultEngine.ini b/Templates/TP_ME_BlankBP/Config/DefaultEngine.ini
--- a/Templates/TP_ME_BlankBP/Config/DefaultEngine.ini
+++ b/Templates/TP_ME_BlankBP/Config/DefaultEngine.ini
@@ -25,7 +25,6 @@ r.AllowStaticLighting=True
 r.HighResScreenshotDelay=8
 r.DefaultBackBufferPixelFormat=4
 r.AllowGlobalClipPlane=False
-r.SSR.Quality=4
 r.GBufferFormat=3
 
 ; shader settings
@@ -35,20 +34,12 @@ r.DefaultFeature.LensFlare=True
 r.ClearCoatNormal=False
 r.NormalMapsForStaticLighting=False
 
->>>> ORIGINAL //UE4/Release-4.24/Templates/TP_ME_BlankBP/Config/DefaultEngine.ini#6
-; Sequencer Take Recorder
-Sequencer.AlwayRecordLiveLinkTimecode 1
-==== THEIRS //UE4/Release-4.24/Templates/TP_ME_BlankBP/Config/DefaultEngine.ini#7
-; Sequencer Take Recorder
-Sequencer.AlwayRecordLiveLinkTimecode=1
-==== YOURS //andy.blondin_4-24_Release/Templates/TP_ME_BlankBP/Config/DefaultEngine.ini
-<<<<
 
 
 [/Script/EngineSettings.GameMapsSettings]
 ;defines what level to be loaded by default in the editor and at playback time
-GameDefaultMap=/Game/Ndisplay.Ndisplay
-EditorStartupMap=/Game/Main.Main
+GameDefaultMap=/Game/VProdProject/Maps/Ndisplay.Ndisplay
+EditorStartupMap=/Game/VProdProject/Maps/Main.Main
 
 [/Script/UnrealEd.UnrealEdEngine]
 ; Remove Engine Template maps
diff --git a/Templates/TP_ME_VProdBP/Config/DefaultEngine.ini b/Templates/TP_ME_VProdBP/Config/DefaultEngine.ini
--- a/Templates/TP_ME_VProdBP/Config/DefaultEngine.ini
+++ b/Templates/TP_ME_VProdBP/Config/DefaultEngine.ini
@@ -33,7 +33,6 @@ r.AllowStaticLighting=True
 r.HighResScreenshotDelay=8
 r.DefaultBackBufferPixelFormat=4
 r.AllowGlobalClipPlane=False
-r.SSR.Quality=4
 r.GBufferFormat=3
 
 ; shader settings
@@ -43,20 +42,12 @@ r.DefaultFeature.LensFlare=False
 r.ClearCoatNormal=False
 r.NormalMapsForStaticLighting=False
 
->>>> ORIGINAL //UE4/Release-4.24/Templates/TP_ME_VProdBP/Config/DefaultEngine.ini#3
-; Sequencer Take Recorder
-Sequencer.AlwayRecordLiveLinkTimecode 1
-==== THEIRS //UE4/Release-4.24/Templates/TP_ME_VProdBP/Config/DefaultEngine.ini#4
-; Sequencer Take Recorder
-Sequencer.AlwayRecordLiveLinkTimecode=1
-==== YOURS //andy.blondin_4-24_Release/Templates/TP_ME_VProdBP/Config/DefaultEngine.ini
-<<<<
 
 
 [/Script/EngineSettings.GameMapsSettings]
 ;defines what level to be loaded by default in the editor and at playback time
-GameDefaultMap=/Game/VProdProject/Ndisplay.Ndisplay
-EditorStartupMap=/Game/VProdProject/Main.Main
+GameDefaultMap=/Game/VProdProject/Maps/Ndisplay.Ndisplay
+EditorStartupMap=/Game/VProdProject/Maps/Main.Main
 
 [/Script/UnrealEd.UnrealEdEngine]
 ; Remove Engine Template maps
