Virutal Production Project Template.  Fixed Spelling errors in TemplateDefs.

diff --git a/Templates/TP_ME_VProdBP/Config/TemplateDefs.ini b/Templates/TP_ME_VProdBP/Config/TemplateDefs.ini
--- a/Templates/TP_ME_VProdBP/Config/TemplateDefs.ini
+++ b/Templates/TP_ME_VProdBP/Config/TemplateDefs.ini
@@ -11,7 +11,7 @@ Categories=ME
 LocalizedDisplayNames=(Language="en", Text="Virtual Production")
 
 ; The localized descriptions for this template, which is displayed when the template is selected.
-LocalizedDescriptions=(Language="en", Text="This Virtual Production project template contains functionality for VR Scouting, Virutal Camera, SDI Video, Live Compositing, and Ndisplay.")
+LocalizedDescriptions=(Language="en", Text="This Virtual Production project template contains functionality for VR Scouting, Virtual Camera, SDI Video, Live Compositing, and Ndisplay.")
 ; The sort key to use for the template to bypass the default alphabetical sorting.
 SortKey="_2"
 
