adding Chinese translation to this file #jira ueloc-4711 #fyi jamie.dale, michelle.zhao, alexis.zhang, yali.he

diff --git a/Templates/TP_ME_VProdBP/Config/TemplateDefs.ini b/Templates/TP_ME_VProdBP/Config/TemplateDefs.ini
--- a/Templates/TP_ME_VProdBP/Config/TemplateDefs.ini
+++ b/Templates/TP_ME_VProdBP/Config/TemplateDefs.ini
@@ -16,6 +16,9 @@ LocalizedDescriptions=(Language="en", Text="This Virtual Production project temp
 LocalizedDisplayNames=(Language="ko", Text="버추얼 프로덕션")
 LocalizedDescriptions=(Language="ko", Text="이 버추얼 프로덕션 프로젝트 템플릿은 VR 스카우팅, 버추얼 카메라, SDI 비디오, 라이브 컴포지팅 그리고 nDisplay를 위한 기능들이 포함되어 있습니다.")
 
+LocalizedDisplayNames=(Language="zh-Hans", Text="虚拟制片")
+LocalizedDescriptions=(Language="zh-Hans", Text="该虚拟制片项目模板包含VR探查、虚拟摄像机、SDI视频、现场合成和nDisplay功能。")
+
 ; The sort key to use for the template to bypass the default alphabetical sorting.
 SortKey="_2"
 
