Adding Japanese translation #jira ueloc-4712 #fyi michelle.zhao, yali.he

diff --git a/Templates/TP_ME_VProdBP/Config/TemplateDefs.ini b/Templates/TP_ME_VProdBP/Config/TemplateDefs.ini
--- a/Templates/TP_ME_VProdBP/Config/TemplateDefs.ini
+++ b/Templates/TP_ME_VProdBP/Config/TemplateDefs.ini
@@ -19,6 +19,8 @@ LocalizedDescriptions=(Language="ko", Text="이 버추얼 프로덕션 프로젝
 LocalizedDisplayNames=(Language="zh-Hans", Text="虚拟制片")
 LocalizedDescriptions=(Language="zh-Hans", Text="该虚拟制片项目模板包含VR探查、虚拟摄像机、SDI视频、现场合成和nDisplay功能。")
 
+LocalizedDescriptions=(Language="ja", Text="この Virtual Production プロジェクト テンプレートには、VR Scouting、Virtual Camera、SDI Video、Live Compositing、nDisplay のための機能が含まれています。")
+
 ; The sort key to use for the template to bypass the default alphabetical sorting.
 SortKey="_2"
 
