Remove hardwired DX11 RHI setting from some templates.

diff --git a/Templates/TP_ME_BlankBP/Config/DefaultEngine.ini b/Templates/TP_ME_BlankBP/Config/DefaultEngine.ini
--- a/Templates/TP_ME_BlankBP/Config/DefaultEngine.ini
+++ b/Templates/TP_ME_BlankBP/Config/DefaultEngine.ini
@@ -90,12 +90,4 @@ EndpointSettings=(bEnableLogging=False,PurgeProcessedMessageDelaySeconds=30,Remo
 ; Setup for Ndisplay
 [/Script/DisplayClusterEditor.DisplayClusterEditorSettings]
 bEnabled=True
-[/Script/WindowsTargetPlatform.WindowsTargetSettings]
-;Ndisplay Requirements
-DefaultGraphicsRHI=DefaultGraphicsRHI_DX11
-+TargetedRHIs=PCD3D_SM5
--TargetedRHIs=PCD3D_SM4
--TargetedRHIs=GLSL_430
--TargetedRHIs=GLSL_150
-
 
diff --git a/Templates/TP_ME_VProdBP/Config/DefaultEngine.ini b/Templates/TP_ME_VProdBP/Config/DefaultEngine.ini
--- a/Templates/TP_ME_VProdBP/Config/DefaultEngine.ini
+++ b/Templates/TP_ME_VProdBP/Config/DefaultEngine.ini
@@ -5,14 +5,6 @@ AppliedTargetedHardwareClass=Desktop
 DefaultGraphicsPerformance=Maximum
 AppliedDefaultGraphicsPerformance=Maximum
 
-[/Script/WindowsTargetPlatform.WindowsTargetSettings]
-;Ndisplay Requirements
-DefaultGraphicsRHI=DefaultGraphicsRHI_DX11
-+TargetedRHIs=PCD3D_SM5
--TargetedRHIs=PCD3D_SM4
--TargetedRHIs=GLSL_430
--TargetedRHIs=GLSL_150
-
 [/Script/Engine.RendererSettings]
 ; exposure control and tone mapping settings
 r.Tonemapper.Sharpen=2
@@ -99,12 +91,3 @@ EndpointSettings=(bEnableLogging=False,PurgeProcessedMessageDelaySeconds=30,Remo
 [/Script/DisplayClusterEditor.DisplayClusterEditorSettings]
 bEnabled=True
 
-[/Script/WindowsTargetPlatform.WindowsTargetSettings]
-;Ndisplay Requirements
-DefaultGraphicsRHI=DefaultGraphicsRHI_DX11
-+TargetedRHIs=PCD3D_SM5
--TargetedRHIs=PCD3D_SM4
--TargetedRHIs=GLSL_430
--TargetedRHIs=GLSL_150
-
-
diff --git a/Templates/TP_nDisplayBP/Config/DefaultEngine.ini b/Templates/TP_nDisplayBP/Config/DefaultEngine.ini
--- a/Templates/TP_nDisplayBP/Config/DefaultEngine.ini
+++ b/Templates/TP_nDisplayBP/Config/DefaultEngine.ini
@@ -20,10 +20,3 @@ vr.AllowMotionBlurInVR=True
 r.DefaultFeature.AutoExposure.Method=1
 r.DefaultFeature.AutoExposure.Bias=1.000000
 
-[/Script/WindowsTargetPlatform.WindowsTargetSettings]
-DefaultGraphicsRHI=DefaultGraphicsRHI_DX11
-+TargetedRHIs=PCD3D_SM5
--TargetedRHIs=PCD3D_SM4
--TargetedRHIs=GLSL_430
--TargetedRHIs=GLSL_150
-
