Dramatically reduced the size of TP_ME_BlankBP and TP_ME_VProdBP on disk by removing unnecessary maps and switching the Main map to no longer need BuiltData.  Left Media assets for now because they do not take up much disk space and the templates are likely to be restructured in 5.7.  Updated DefaultInput.ini to use Enhanced Input to avoid startup messages of it being automatically overridden.

#virtualized

diff --git a/Templates/TP_ME_BlankBP/Config/DefaultInput.ini b/Templates/TP_ME_BlankBP/Config/DefaultInput.ini
--- a/Templates/TP_ME_BlankBP/Config/DefaultInput.ini
+++ b/Templates/TP_ME_BlankBP/Config/DefaultInput.ini
@@ -1,5 +1,4 @@
-
-
 [/Script/Engine.InputSettings]
 bAlwaysShowTouchInterface=True
-
+DefaultPlayerInputClass=/Script/EnhancedInput.EnhancedPlayerInput
+DefaultInputComponentClass=/Script/EnhancedInput.EnhancedInputComponent
\ No newline at end of file
diff --git a/Templates/TP_ME_VProdBP/Config/DefaultEngine.ini b/Templates/TP_ME_VProdBP/Config/DefaultEngine.ini
--- a/Templates/TP_ME_VProdBP/Config/DefaultEngine.ini
+++ b/Templates/TP_ME_VProdBP/Config/DefaultEngine.ini
@@ -39,7 +39,7 @@ r.NormalMapsForStaticLighting=False
 
 [/Script/EngineSettings.GameMapsSettings]
 ;defines what level to be loaded by default in the editor and at playback time
-GameDefaultMap=/Game/VProdProject/Maps/Ndisplay.Ndisplay
+GameDefaultMap=/Game/VProdProject/Maps/Main.Main
 EditorStartupMap=/Game/VProdProject/Maps/Main.Main
 
 [/Script/UnrealEd.UnrealEdEngine]
@@ -48,6 +48,7 @@ EditorStartupMap=/Game/VProdProject/Maps/Main.Main
 
 
 ;Allows for Hardware Accelerated Video Decoding
+
 [/Script/WmfMediaFactory.WmfMediaSettings]
 AllowNonStandardCodecs=True
 HardwareAcceleratedVideoDecoding=True
@@ -55,6 +56,7 @@ LowLatency=False
 NativeAudioOut=False
 
 ;Adds Virtual Scouting Widget for VR Scouting
+
 [/Script/VPUtilitiesEditor.VPUtilitiesEditorSettings]
 VirtualScoutingUI=/VirtualProductionUtilities/Editor/VirtualScoutingWidget.VirtualScoutingWidget_C
 
@@ -62,12 +64,14 @@ VirtualScoutingUI=/VirtualProductionUtilities/Editor/VirtualScoutingWidget.Virtu
 WebControl.EnableServerOnStartup=1
 
 ;Adds Remote Session for Vcam
+
 [RemoteSession]
 +Channels=(Name=FRemoteSessionFrameBufferChannel,Mode=Write)
 +Channels=(Name=FRemoteSessionInputChannel,Mode=Read)
 +Channels=(Name=FRemoteSessionXRTrackingChannel,Mode=Read)
 
 ;Setup for Media Profiles
+
 [/Script/MediaFrameworkUtilities.MediaProfileSettings]
 bApplyInCommandlet=False
 MediaSourceProxy=/Game/VprodProject/MediaIO/MediaSource-01.MediaSource-01
@@ -77,6 +81,7 @@ StartupMediaProfile=None
 
 
 ;Setup for Multiuser
+
 [/Script/Concert.ConcertClientConfig]
 bIsHeadless=False
 bInstallEditorToolbarButton=True
@@ -90,6 +95,7 @@ EndpointSettings=(bEnableLogging=False,PurgeProcessedMessageDelaySeconds=30,Remo
 
 
 ;Setup for Ndisplay Cluster
+
 [/Script/DisplayClusterEditor.DisplayClusterEditorSettings]
 bEnabled=True
 
@@ -103,3 +109,18 @@ DefaultGraphicsRHI=DefaultGraphicsRHI_DX12
 [/Script/LinuxTargetPlatform.LinuxTargetSettings]
 -TargetedRHIs=SF_VULKAN_SM5
 +TargetedRHIs=SF_VULKAN_SM6
+
+[/Script/AndroidFileServerEditor.AndroidFileServerRuntimeSettings]
+bEnablePlugin=True
+bAllowNetworkConnection=True
+SecurityToken=FE1C85674E26FE981BF51B824609C20E
+bIncludeInShipping=False
+bAllowExternalStartInShipping=False
+bCompileAFSProject=False
+bUseCompression=False
+bLogFiles=False
+bReportStats=False
+ConnectionType=USBOnly
+bUseManualIPAddress=False
+ManualIPAddress=
+
diff --git a/Templates/TP_ME_VProdBP/Config/DefaultInput.ini b/Templates/TP_ME_VProdBP/Config/DefaultInput.ini
--- a/Templates/TP_ME_VProdBP/Config/DefaultInput.ini
+++ b/Templates/TP_ME_VProdBP/Config/DefaultInput.ini
@@ -0,0 +1,4 @@
+[/Script/Engine.InputSettings]
+bAlwaysShowTouchInterface=True
+DefaultPlayerInputClass=/Script/EnhancedInput.EnhancedPlayerInput
+DefaultInputComponentClass=/Script/EnhancedInput.EnhancedInputComponent
\ No newline at end of file
