[VP Project Templates]
- Enabled Performance Capture Workflow plugin for both the blank M&E and VP project templates
- Modified description text strings to reflect the presence of these additional plugins.

#Synced-CL 51579190

diff --git a/Templates/TP_ME_BlankBP/Config/TemplateDefs.ini b/Templates/TP_ME_BlankBP/Config/TemplateDefs.ini
--- a/Templates/TP_ME_BlankBP/Config/TemplateDefs.ini
+++ b/Templates/TP_ME_BlankBP/Config/TemplateDefs.ini
@@ -11,7 +11,7 @@ Categories=ME
 LocalizedDisplayNames=(Language="en", Text="Blank")
 
 ; The localized descriptions for this template, which is displayed when the template is selected.
-LocalizedDescriptions=(Language="en", Text="Blank template with minimal functionality.")
+LocalizedDescriptions=(Language="en", Text="Blank Virtual Production template with minimal content.")
 
 LocalizedDisplayNames=(Language="ko", Text="빈 프로젝트")
 LocalizedDescriptions=(Language="ko", Text="최소한의 기능을 가진 빈 프로젝트 템플릿")
diff --git a/Templates/TP_ME_BlankBP/TP_ME_BlankBP.uproject b/Templates/TP_ME_BlankBP/TP_ME_BlankBP.uproject
--- a/Templates/TP_ME_BlankBP/TP_ME_BlankBP.uproject
+++ b/Templates/TP_ME_BlankBP/TP_ME_BlankBP.uproject
@@ -134,6 +134,10 @@
 		{
 			"Name": "MovieRenderPipeline",
 			"Enabled": true
+		},
+		{
+			"Name": "PerformanceCaptureWorkflow",
+			"Enabled": true
 		}
 	]
 }
\ No newline at end of file
diff --git a/Templates/TP_ME_VProdBP/Config/TemplateDefs.ini b/Templates/TP_ME_VProdBP/Config/TemplateDefs.ini
--- a/Templates/TP_ME_VProdBP/Config/TemplateDefs.ini
+++ b/Templates/TP_ME_VProdBP/Config/TemplateDefs.ini
@@ -11,7 +11,7 @@ Categories=ME
 LocalizedDisplayNames=(Language="en", Text="Virtual Production")
 
 ; The localized descriptions for this template, which is displayed when the template is selected.
-LocalizedDescriptions=(Language="en", Text="This Virtual Production project template contains functionality for VR Scouting, Virtual Camera, SDI Video, Live Compositing, and nDisplay.")
+LocalizedDescriptions=(Language="en", Text="This Virtual Production project template contains functionality for VR Scouting, Virtual Camera, Performance Capture, SDI Video, Live Compositing, and nDisplay.")
 
 LocalizedDisplayNames=(Language="ko", Text="버추얼 프로덕션")
 LocalizedDescriptions=(Language="ko", Text="이 버추얼 프로덕션 프로젝트 템플릿은 VR 스카우팅, 버추얼 카메라, SDI 비디오, 라이브 컴포지팅 그리고 nDisplay를 위한 기능들이 포함되어 있습니다.")
diff --git a/Templates/TP_ME_VProdBP/TP_ME_VprodBP.uproject b/Templates/TP_ME_VProdBP/TP_ME_VprodBP.uproject
--- a/Templates/TP_ME_VProdBP/TP_ME_VprodBP.uproject
+++ b/Templates/TP_ME_VProdBP/TP_ME_VprodBP.uproject
@@ -134,6 +134,10 @@
 		{
 			"Name": "MovieRenderPipeline",
 			"Enabled": true
+		},
+		{
+			"Name": "PerformanceCaptureWorkflow",
+			"Enabled": true
 		}
 	]
 }
\ No newline at end of file
